Added UnrealWatchdog tool, run by the Editor, to improve abnormal shutdown tracking.
[UE-32952] - Watchdog - Show CRC when reporting abnormal shutdowns in internal builds

Reimplemented in //UE4/Main from //UE4/Dev-Editor

diff --git a/Engine/Config/BaseEngine.ini b/Engine/Config/BaseEngine.ini
--- a/Engine/Config/BaseEngine.ini
+++ b/Engine/Config/BaseEngine.ini
@@ -2807,6 +2807,7 @@ ConnectionRetryDelay=2
 [CrashReportClient]
 bHideLogFilesOption=false
 bIsAllowedToCloseWithoutSending=true
+
 [SteamVR.Settings]
 HMDWornMovementThreshold = 50.0
 
@@ -2815,3 +2816,6 @@ bStripAnimationDataOnDedicatedServer=False
 
 [/Script/Engine.MeshSimplificationSettings]
 r.MeshReductionModule="QuadricMeshReduction"
+
+[EngineSessionManager]
+UseWatchdogMTBF=false
\ No newline at end of file
diff --git a/Engine/Source/Programs/UnrealWatchdog/Private/Linux/UnrealWatchdogMainLinux.cpp b/Engine/Source/Programs/UnrealWatchdog/Private/Linux/UnrealWatchdogMainLinux.cpp
new file mode 100644
--- /dev/null
+++ b/Engine/Source/Programs/UnrealWatchdog/Private/Linux/UnrealWatchdogMainLinux.cpp
@@ -0,0 +1,66 @@
+// Copyright 1998-2016 Epic Games, Inc. All Rights Reserved.
+
+#include "UnrealWatchdog.h"
+#include "LinuxPlatformCrashContext.h"
+#include <locale.h>
+
+extern int32 ReportCrash(const FLinuxCrashContext& Context);	// FIXME: handle expose it someplace else?
+
+																/**
+																* Because crash reporters can crash, too
+																*/
+void UnrealWatchdogCrashHandler(const FGenericCrashContext& GenericContext)
+{
+	// at this point we should already be using malloc crash handler (see PlatformCrashHandler)
+
+	const FLinuxCrashContext& Context = static_cast< const FLinuxCrashContext& >(GenericContext);
+
+	printf("CrashHandler: Signal=%d\n", Context.Signal);
+	const_cast< FLinuxCrashContext& >(Context).CaptureStackTrace();
+	if (GLog)
+	{
+		GLog->Flush();
+	}
+	if (GWarn)
+	{
+		GWarn->Flush();
+	}
+	if (GError)
+	{
+		GError->Flush();
+		GError->HandleError();
+	}
+	FPlatformMisc::RequestExit(true);
+}
+
+
+/**
+* main(), called when the application is started
+*/
+int main(int argc, const char *argv[])
+{
+	FPlatformMisc::SetGracefulTerminationHandler();
+	FPlatformMisc::SetCrashHandler(UnrealWatchdogCrashHandler);
+
+	FString SavedCommandLine;
+
+	setlocale(LC_CTYPE, "");
+	for (int32 Option = 1; Option < argc; Option++)
+	{
+		SavedCommandLine += TEXT(" ");
+		SavedCommandLine += UTF8_TO_TCHAR(argv[Option]);	// note: technically it depends on locale
+	}
+
+	// Run the app
+	RunUnrealWatchdog(*SavedCommandLine);
+
+	return 0;
+}
+
+bool WaitForProcess(IAnalyticsProviderET& Analytics, const FWatchdogCommandLine& CommandLine, int32& OutReturnCode, bool& bOutHang, const FString& WatchdogSectionName)
+{
+	// @TODO: Chris.Wood
+	// See Mac version for comments
+
+	return false;
+}
\ No newline at end of file
diff --git a/Engine/Source/Programs/UnrealWatchdog/Private/Mac/UnrealWatchdogMainMac.cpp b/Engine/Source/Programs/UnrealWatchdog/Private/Mac/UnrealWatchdogMainMac.cpp
new file mode 100644
--- /dev/null
+++ b/Engine/Source/Programs/UnrealWatchdog/Private/Mac/UnrealWatchdogMainMac.cpp
@@ -0,0 +1,117 @@
+// Copyright 1998-2016 Epic Games, Inc. All Rights Reserved.
+
+#include "UnrealWatchdog.h"
+#include "ExceptionHandling.h"
+#include "MacPlatformCrashContext.h"
+#include "CocoaThread.h"
+
+/**
+* Because crash reporters can crash, too - only used for Sandboxed applications
+*/
+void UnrealWatchdogCrashHandler(const FGenericCrashContext& GenericContext)
+{
+	// We never emit a crash report in Sandboxed builds from CRC as if we do, then the crashed application's
+	// crash report is overwritten by the CRC's when trampolining into the Apple Crash Reporter.
+	_Exit(0);
+}
+
+static FString GSavedCommandLine;
+
+@interface UE4AppDelegate : NSObject <NSApplicationDelegate, NSFileManagerDelegate>
+{
+}
+
+@end
+
+@implementation UE4AppDelegate
+
+//handler for the quit apple event used by the Dock menu
+- (void)handleQuitEvent:(NSAppleEventDescriptor*)Event withReplyEvent : (NSAppleEventDescriptor*)ReplyEvent
+{
+	[self requestQuit : self];
+}
+
+-(IBAction)requestQuit : (id)Sender
+{
+	GIsRequestingExit = true;
+}
+
+-(void)runGameThread : (id)Arg
+{
+	FPlatformMisc::SetGracefulTerminationHandler();
+	// For sandboxed applications CRC can never report a crash, or we break trampolining into Apple's crash reporter.
+	if (FPlatformProcess::IsSandboxedApplication())
+	{
+		// @TODO!!
+		FPlatformMisc::SetCrashHandler(CrashReporterCrashHandler);
+	}
+
+	// Run the app
+	RunUnrealWatchdog(*GSavedCommandLine);
+
+	[NSApp terminate : self];
+}
+
+-(NSApplicationTerminateReply)applicationShouldTerminate:(NSApplication *)Sender;
+{
+	if (!GIsRequestingExit || ([NSThread gameThread] && [NSThread gameThread] != [NSThread mainThread]))
+	{
+		[self requestQuit : self];
+		return NSTerminateLater;
+	}
+	else
+	{
+		return NSTerminateNow;
+	}
+}
+
+-(void)applicationDidFinishLaunching:(NSNotification *)Notification
+{
+	//install the custom quit event handler
+	NSAppleEventManager* appleEventManager = [NSAppleEventManager sharedAppleEventManager];
+	[appleEventManager setEventHandler : self andSelector : @selector(handleQuitEvent : withReplyEvent : ) forEventClass:kCoreEventClass andEventID : kAEQuitApplication];
+
+	FPlatformMisc::ActivateApplication();
+	RunGameThread(self, @selector(runGameThread:));
+}
+
+@end
+
+INT32_MAIN_INT32_ARGC_TCHAR_ARGV()
+{
+	for (int32 Option = 1; Option < ArgC; Option++)
+	{
+		GSavedCommandLine += TEXT(" ");
+		FString Argument(ArgV[Option]);
+		if (Argument.Contains(TEXT(" ")))
+		{
+			if (Argument.Contains(TEXT("=")))
+			{
+				FString ArgName;
+				FString ArgValue;
+				Argument.Split(TEXT("="), &ArgName, &ArgValue);
+				Argument = FString::Printf(TEXT("%s=\"%s\""), *ArgName, *ArgValue);
+			}
+			else
+			{
+				Argument = FString::Printf(TEXT("\"%s\""), *Argument);
+			}
+		}
+		GSavedCommandLine += Argument;
+	}
+
+	SCOPED_AUTORELEASE_POOL;
+	[NSApplication sharedApplication];
+	[NSApp setDelegate : [UE4AppDelegate new]];
+	[NSApp run];
+	return 0;
+}
+
+bool WaitForProcess(IAnalyticsProviderET& Analytics, const FWatchdogCommandLine& CommandLine, int32& OutReturnCode, bool& bOutHang, const FString& WatchdogSectionName)
+{
+	// @TODO: Chris.Wood
+	// https://developer.apple.com/library/content/technotes/tn2050/_index.html#//apple_ref/doc/uid/DTS10003081-CH1-SECTION8
+	// Use kqueue to get process exit code from the PID on Mac
+
+	return false;
+}
\ No newline at end of file
diff --git a/Engine/Source/Programs/UnrealWatchdog/Private/UnrealWatchdog.cpp b/Engine/Source/Programs/UnrealWatchdog/Private/UnrealWatchdog.cpp
new file mode 100644
--- /dev/null
+++ b/Engine/Source/Programs/UnrealWatchdog/Private/UnrealWatchdog.cpp
@@ -0,0 +1,476 @@
+// Copyright 1998-2016 Epic Games, Inc. All Rights Reserved.
+
+#include "UnrealWatchdog.h"
+#include "RequiredProgramMainCPPInclude.h"
+#include "IAnalyticsProviderET.h"
+#include "WatchdogAnalytics.h"
+#include "ExceptionHandling.h"
+
+IMPLEMENT_APPLICATION(UnrealWatchdog, "UnrealWatchdog");
+DEFINE_LOG_CATEGORY(UnrealWatchdogLog);
+
+#define LOCTEXT_NAMESPACE "UnrealWatchdog"
+
+namespace WatchdogDefs
+{
+	static const FString StoreId(TEXT("Epic Games"));
+	static const FString WatchdogRecordSectionPrefix(TEXT("Unreal Engine/Watchdog/"));
+	static const FString WatchdogVersionString(TEXT("1_0"));
+	static const FString CommandLineStoreKey(TEXT("CommandLine"));
+	static const FString StartupTimeStoreKey(TEXT("StartupTimestamp"));
+	static const FString TimestampStoreKey(TEXT("Timestamp"));
+	static const FString StatusStoreKey(TEXT("LastExecutionState"));
+	static const FString UserActivityStoreKey(TEXT("CurrentUserActivity"));
+	static const FString RunningSessionToken(TEXT("Running"));
+	static const FString ShutdownSessionToken(TEXT("Shutdown"));
+	static const FString CrashSessionToken(TEXT("Crashed"));
+}
+
+struct FWatchdogStoredValues
+{
+	FString CommandLine;
+	FString StartTime;
+	FString LastTimestamp;
+	FString ExecutionStatus;
+	FString UserActivity;
+};
+
+void GetCommonEventAttributes(const FWatchdogCommandLine& CommandLine, TArray< FAnalyticsEventAttribute >& OutAttributes)
+{
+	OutAttributes.Add(FAnalyticsEventAttribute(TEXT("RunType"), CommandLine.RunType));
+	OutAttributes.Add(FAnalyticsEventAttribute(TEXT("ProjectName"), CommandLine.ProjectName));
+	OutAttributes.Add(FAnalyticsEventAttribute(TEXT("Platform"), CommandLine.PlatformName));
+	OutAttributes.Add(FAnalyticsEventAttribute(TEXT("SessionId"), CommandLine.SessionId));
+	OutAttributes.Add(FAnalyticsEventAttribute(TEXT("EngineVersion"), CommandLine.EngineVersion));
+}
+
+FWatchdogStoredValues GetWatchdogStoredValues(const FString& WatchdogSectionName)
+{
+	FWatchdogStoredValues StoredValues;
+
+	FPlatformMisc::GetStoredValue(WatchdogDefs::StoreId, WatchdogSectionName, WatchdogDefs::CommandLineStoreKey, StoredValues.CommandLine);
+	FPlatformMisc::GetStoredValue(WatchdogDefs::StoreId, WatchdogSectionName, WatchdogDefs::StartupTimeStoreKey, StoredValues.StartTime);
+	FPlatformMisc::GetStoredValue(WatchdogDefs::StoreId, WatchdogSectionName, WatchdogDefs::TimestampStoreKey, StoredValues.LastTimestamp);
+	FPlatformMisc::GetStoredValue(WatchdogDefs::StoreId, WatchdogSectionName, WatchdogDefs::StatusStoreKey, StoredValues.ExecutionStatus);
+	FPlatformMisc::GetStoredValue(WatchdogDefs::StoreId, WatchdogSectionName, WatchdogDefs::UserActivityStoreKey, StoredValues.UserActivity);
+
+	FPlatformMisc::DeleteStoredValue(WatchdogDefs::StoreId, WatchdogSectionName, WatchdogDefs::CommandLineStoreKey);
+	FPlatformMisc::DeleteStoredValue(WatchdogDefs::StoreId, WatchdogSectionName, WatchdogDefs::StartupTimeStoreKey);
+	FPlatformMisc::DeleteStoredValue(WatchdogDefs::StoreId, WatchdogSectionName, WatchdogDefs::TimestampStoreKey);
+	FPlatformMisc::DeleteStoredValue(WatchdogDefs::StoreId, WatchdogSectionName, WatchdogDefs::StatusStoreKey);
+	FPlatformMisc::DeleteStoredValue(WatchdogDefs::StoreId, WatchdogSectionName, WatchdogDefs::UserActivityStoreKey);
+
+	return StoredValues;
+}
+
+FString GetWatchdogStoredTimestamp(const FString& WatchdogSectionName)
+{
+	FString Timestamp;
+	FPlatformMisc::GetStoredValue(WatchdogDefs::StoreId, WatchdogSectionName, WatchdogDefs::TimestampStoreKey, Timestamp);
+	return Timestamp;
+}
+
+/**
+ * @EventName UnrealWatchdog.Initialized
+ *
+ * @Trigger Event raised by UnrealWatchdog at startup. Records an instance of the watchdog running and whether the watched process was found successfully.
+ *
+ * @Type Static
+ * @Owner Chris.Wood
+ *
+ * @EventParam RunType - Editor or Game
+ * @EventParam ProjectName - Project for the session that abnormally terminated.
+ * @EventParam Platform - Windows, Mac, Linux
+ * @EventParam SessionId - Analytics SessionID of the session that abnormally terminated.
+ * @EventParam EngineVersion - EngineVersion of the session that abnormally terminated.
+ * @EventParam bValidPID - Was the Watchdog passed a valid PID on the command line? (True or False)
+ * @EventParam bProcessFound - Was the watched process running when the Watchdog started up? (True or False)
+ */
+void SendStartupEvent(IAnalyticsProviderET& Analytics, const FWatchdogCommandLine& CommandLine, bool bInitiallyRunning)
+{
+	TArray< FAnalyticsEventAttribute > StartupAttributes;
+	GetCommonEventAttributes(CommandLine, StartupAttributes);
+	StartupAttributes.Add(FAnalyticsEventAttribute(TEXT("bValidPID"), CommandLine.bHasProcessId ? TEXT("True") : TEXT("False")));
+	StartupAttributes.Add(FAnalyticsEventAttribute(TEXT("bProcessFound"), bInitiallyRunning ? TEXT("True") : TEXT("False")));
+
+	UE_LOG(UnrealWatchdogLog, Log, TEXT("Sending event UnrealWatchdog.Initialized"));
+	Analytics.RecordEvent(TEXT("UnrealWatchdog.Initialized"), MoveTemp(StartupAttributes));
+}
+
+/**
+ * @EventName UnrealWatchdog.Heartbeat
+ *
+ * @Trigger Event raised by a running UnrealWatchdog process while waiting for a process to exit.
+ *
+ * @Type Static
+ * @Owner Chris.Wood
+ *
+ * @EventParam RunType - Editor or Game
+ * @EventParam ProjectName - Project for the session.
+ * @EventParam Platform - Windows, Mac, Linux
+ * @EventParam SessionId - Analytics SessionID of the session.
+ * @EventParam EngineVersion - EngineVersion of the session.
+ */
+void SendHeartbeatEvent(IAnalyticsProviderET& Analytics, const FWatchdogCommandLine& CommandLine)
+{
+	static TArray< FAnalyticsEventAttribute > HeartbeatAttributes;
+	static FDateTime LastEventTime;
+	if (HeartbeatAttributes.Num() == 0)
+	{
+		// Init
+		GetCommonEventAttributes(CommandLine, HeartbeatAttributes);
+		HeartbeatAttributes.Add(FAnalyticsEventAttribute(TEXT("IntervalSec"), 0.0f));
+		LastEventTime = FDateTime::UtcNow();
+	}
+	else
+	{
+		// Update the IntervalSec with the exact time since the last heartbeat
+		FDateTime EventTime = FDateTime::UtcNow();
+		FTimespan ActualInterval = EventTime - LastEventTime;
+		HeartbeatAttributes[HeartbeatAttributes.Num() - 1] = FAnalyticsEventAttribute(TEXT("IntervalSec"), ActualInterval.GetTotalSeconds());
+		LastEventTime = EventTime;
+	}
+
+	UE_LOG(UnrealWatchdogLog, Verbose, TEXT("Sending event UnrealWatchdog.Heartbeat"));
+	Analytics.RecordEvent(TEXT("UnrealWatchdog.Heartbeat"), CopyTemp(HeartbeatAttributes));
+}
+
+bool CheckParentHeartbeat(IAnalyticsProviderET& Analytics, const FWatchdogCommandLine& CommandLine, const FString& WatchdogSectionName)
+{
+	static FString HeartbeatString;
+	FString NextHeartbeatString = GetWatchdogStoredTimestamp(WatchdogSectionName);
+	if (HeartbeatString != NextHeartbeatString)
+	{
+		HeartbeatString = NextHeartbeatString;
+		return true;
+	}
+	return false;
+}
+
+/**
+ * @EventName UnrealWatchdog.HangDetected
+ *
+ * @Trigger Event raised when the watchdog detects an interruption in the heartbeat of the watched process.
+ *
+ * @Type Static
+ * @Owner Chris.Wood
+ *
+ * @EventParam RunType - Editor or Game
+ * @EventParam ProjectName - Project for the session.
+ * @EventParam Platform - Windows, Mac, Linux
+ * @EventParam SessionId - Analytics SessionID of the session.
+ * @EventParam EngineVersion - EngineVersion of the session.
+ * @EventParam HangUserResponse - Indicates the response, if any, from the user when asked about the hang by the watchdog (Unattended, Confirmed, False)
+ * @EventParam AlreadyRecoveredUserResponse - If the user responded by confirming the hang (HangUserResponse=Confirmed), indicates the response when asked if the hang has recovered (Recovered, NotRecovered, N/A)
+ */
+void SendHangDetectedEvent(IAnalyticsProviderET& Analytics, const FWatchdogCommandLine& CommandLine)
+{
+	TArray< FAnalyticsEventAttribute > HangAttributes;
+	GetCommonEventAttributes(CommandLine, HangAttributes);
+
+	// Internal builds should popup dialogs for hangs
+	FAnalyticsEventAttribute HangUserResponse(TEXT("HangUserResponse"), TEXT("Unattended"));
+	FAnalyticsEventAttribute RecoveredUserResponse(TEXT("AlreadyRecoveredUserResponse"), TEXT("Unattended"));
+	if (FEngineBuildSettings::IsInternalBuild())
+	{
+		FText SessionLabel = FText::Format(LOCTEXT("HangSessionLabel", "{0} ({1})"), FText::FromString(CommandLine.ProjectName), FText::FromString(CommandLine.RunType));
+		FText MessageTitleFormat(LOCTEXT("WatchdogPopupTitleHang", "{0} is unresponsive"));
+		FText MessageTitle = FText::Format(MessageTitleFormat, SessionLabel);
+		FText Message(LOCTEXT("WatchdogPopupQuestionHang", "We think the application may be hanging. Did the application freeze without closing or showing the Crash Reporter?"));
+
+		EAppReturnType::Type HangAnswer = FMessageDialog::Open(EAppMsgType::YesNo, Message, &MessageTitle);
+
+		if (HangAnswer == EAppReturnType::Yes)
+		{
+			UE_LOG(UnrealWatchdogLog, Log, TEXT("User confirmed hang"));
+			HangUserResponse.AttrValue = TEXT("Confirmed");
+
+			FText RecoveredMessage(LOCTEXT("WatchdogPopupRecoveredMessage", "Has the application recovered?"));
+
+			EAppReturnType::Type RecoveredAnswer = FMessageDialog::Open(EAppMsgType::YesNo, RecoveredMessage, &MessageTitle);
+
+			if (RecoveredAnswer == EAppReturnType::Yes)
+			{
+				UE_LOG(UnrealWatchdogLog, Log, TEXT("User confirmed recovery from hang"));
+				RecoveredUserResponse.AttrValue = TEXT("Recovered");
+			}
+			else
+			{
+				UE_LOG(UnrealWatchdogLog, Log, TEXT("User confirmed hang not yet recovered"));
+				RecoveredUserResponse.AttrValue = TEXT("NotRecovered");
+			}
+
+			FText CRCMessage(LOCTEXT("WatchdogPopupHangCRCMessage", "We will now open the Crash Reporter for you to tell us what happened."));
+			FMessageDialog::Open(EAppMsgType::Ok, CRCMessage, &MessageTitle);
+
+			FString EnsureText = FString::Printf(TEXT("Watchdog detected hang in %s."), *SessionLabel.ToString());
+			ReportInteractiveEnsure(*EnsureText);
+		}
+		else  // HangAnswer == EAppReturnType::No
+		{
+			UE_LOG(UnrealWatchdogLog, Warning, TEXT("User didn't witness hang. False positive warning!"));
+			HangUserResponse.AttrValue = TEXT("False");
+			RecoveredUserResponse.AttrValue = TEXT("N/A");
+		}
+	}
+
+	HangAttributes.Add(HangUserResponse);
+	HangAttributes.Add(RecoveredUserResponse);
+	UE_LOG(UnrealWatchdogLog, Log, TEXT("Sending event UnrealWatchdog.HangDetected"));
+	Analytics.RecordEvent(TEXT("UnrealWatchdog.HangDetected"), MoveTemp(HangAttributes));
+}
+
+/**
+ * @EventName UnrealWatchdog.HangRecovered
+ *
+ * @Trigger Event raised when the watchdog detects that the heartbeat of the watched process has resumed updating after a hang was detected.
+ *
+ * @Type Static
+ * @Owner Chris.Wood
+ *
+ * @EventParam RunType - Editor or Game
+ * @EventParam ProjectName - Project for the session.
+ * @EventParam Platform - Windows, Mac, Linux
+ * @EventParam SessionId - Analytics SessionID of the session.
+ * @EventParam EngineVersion - EngineVersion of the session.
+ */
+void SendHangRecoveredEvent(IAnalyticsProviderET& Analytics, const FWatchdogCommandLine& CommandLine)
+{
+	TArray< FAnalyticsEventAttribute > HangAttributes;
+	GetCommonEventAttributes(CommandLine, HangAttributes);
+
+	UE_LOG(UnrealWatchdogLog, Log, TEXT("Sending event UnrealWatchdog.HangRecovered"));
+	Analytics.RecordEvent(TEXT("UnrealWatchdog.HangRecovered"), MoveTemp(HangAttributes));
+}
+
+/**
+ * @EventName UnrealWatchdog.Shutdown
+ *
+ * @Trigger Event raised by UnrealWatchdog at shutdown. The watched process has exited (or wasn't found).
+ *
+ * @Type Static
+ * @Owner Chris.Wood
+ *
+ * @EventParam RunType - Editor or Game
+ * @EventParam ProjectName - Project for the session that abnormally terminated.
+ * @EventParam Platform - Windows, Mac, Linux
+ * @EventParam SessionId - Analytics SessionID of the session that abnormally terminated.
+ * @EventParam EngineVersion - EngineVersion of the session that abnormally terminated.
+ * @EventParam bReturnCodeObtained - Did the Watchdog get the return/exit code from the process? (True or False)
+ * @EventParam OSReturnCode - The exit code of the watched process.
+ * @EventParam CommandLine - Full command line that ran the watched process.
+ * @EventParam StartTime - The time the watched process initialized the watchdog.
+ * @EventParam LastTimestamp - The last time the watched process updated its stored status (not necessarily the shutdown time).
+ * @EventParam LastExecutionStatus - The last updated status of the watched process when it was running. (Running, Shutdown, Crashed or Hang)
+ * @EventParam LastUserActivity - The last updated user activity, if any, of the watched process.
+ * @EventParam AbnormalShutdownUserResponse - Indicates the response, if any, from the user when asked about the shutdown by the watchdog (Unattended, Confirmed, False)
+ */
+void SendShutdownEvent(IAnalyticsProviderET& Analytics, const FWatchdogCommandLine& CommandLine, bool bReturnCodeObtained, uint32 ReturnCode, FAnalyticsEventAttribute& UserResponse, const FWatchdogStoredValues& StoredValues)
+{
+	TArray< FAnalyticsEventAttribute > ShutdownAttributes;
+	GetCommonEventAttributes(CommandLine, ShutdownAttributes);
+	ShutdownAttributes.Add(FAnalyticsEventAttribute(TEXT("bReturnCodeObtained"), bReturnCodeObtained ? TEXT("True") : TEXT("False")));
+	ShutdownAttributes.Add(FAnalyticsEventAttribute(TEXT("OSReturnCode"), ReturnCode));
+	ShutdownAttributes.Add(FAnalyticsEventAttribute(TEXT("CommandLine"), StoredValues.CommandLine));
+	ShutdownAttributes.Add(FAnalyticsEventAttribute(TEXT("StartTime"), StoredValues.StartTime));
+	ShutdownAttributes.Add(FAnalyticsEventAttribute(TEXT("LastTimestamp"), StoredValues.LastTimestamp));
+	ShutdownAttributes.Add(FAnalyticsEventAttribute(TEXT("LastExecutionStatus"), StoredValues.ExecutionStatus));
+	ShutdownAttributes.Add(FAnalyticsEventAttribute(TEXT("LastUserActivity"), StoredValues.UserActivity));
+	ShutdownAttributes.Add(UserResponse);
+
+	UE_LOG(UnrealWatchdogLog, Log, TEXT("Sending event UnrealWatchdog.Shutdown"));
+	Analytics.RecordEvent(TEXT("UnrealWatchdog.Shutdown"), MoveTemp(ShutdownAttributes));
+}
+
+/**
+* @EventName UnrealWatchdog.AbnormalShutdown
+*
+* @Trigger Event raised by UnrealWatchdog as part of MTBF tracking. Records an instance of a watchdog protected process shutting down abnormally.
+*
+* @Type Static
+* @Owner Chris.Wood
+*
+* @EventParam RunType - Editor or Game
+* @EventParam ProjectName - Project for the session that abnormally terminated.
+* @EventParam Platform - Windows, Mac, Linux
+* @EventParam SessionId - Analytics SessionID of the session that abnormally terminated.
+* @EventParam EngineVersion - EngineVersion of the session that abnormally terminated.
+* @EventParam UserResponse - What did the user click if and when the watchdog showed a message box. One of Confirmed, False or Unattended.
+*/
+int RunUnrealWatchdog(const TCHAR* CommandLine)
+{
+	// start up the main loop
+	GEngineLoop.PreInit(CommandLine);
+
+	check(GConfig && GConfig->IsReadyForUse());
+
+	FWatchdogCommandLine WatchdogCommandLine(CommandLine);
+	bool bInitiallyRunning = WatchdogCommandLine.bHasProcessId && FPlatformProcess::IsApplicationRunning(WatchdogCommandLine.ParentProcessId);
+
+	// Send watchdog startup event
+	FWatchdogAnalytics::Initialize();
+	IAnalyticsProviderET& Analytics = FWatchdogAnalytics::GetProvider();
+	SendStartupEvent(Analytics, WatchdogCommandLine, bInitiallyRunning);
+
+	if (!WatchdogCommandLine.bHasProcessId)
+	{
+		UE_LOG(UnrealWatchdogLog, Error, TEXT("Watchdog wasn't given a valid PID"));
+	}
+	else if (!bInitiallyRunning)
+	{
+		UE_LOG(UnrealWatchdogLog, Warning, TEXT("Watchdog didn't find running process PID %u..."), WatchdogCommandLine.ParentProcessId);
+	}
+	else
+	{
+		UE_LOG(UnrealWatchdogLog, Log, TEXT("Watchdog beginning to watch PID %u..."), WatchdogCommandLine.ParentProcessId);
+	}
+
+	FString WatchdogSectionName = FString::Printf(TEXT("%s%s/%u"), *WatchdogDefs::WatchdogRecordSectionPrefix, *WatchdogDefs::WatchdogVersionString, WatchdogCommandLine.ParentProcessId);
+
+	// Wait for the process to exit
+	int32 ReturnCode = WatchdogCommandLine.SuccessReturnCode;
+	bool bHang = false;
+	bool bReturnCodeObtained = WaitForProcess(Analytics, WatchdogCommandLine, ReturnCode, bHang, WatchdogSectionName);
+
+	// Read any stored values from process
+	FWatchdogStoredValues StoredValues = GetWatchdogStoredValues(WatchdogSectionName);
+	if (bHang)
+	{
+		StoredValues.ExecutionStatus = TEXT("Hang");
+	}
+
+	// Optional section for dialogs and CRC in internal builds
+	FAnalyticsEventAttribute UserResponse(TEXT("AbnormalShutdownUserResponse"), TEXT("Unattended"));
+	if (WatchdogCommandLine.bAllowDialogs && !bHang)
+	{
+		EAppReturnType::Type UserAnswer = EAppReturnType::Cancel;
+		FText SessionLabel = FText::Format(LOCTEXT("SessionLabel", "{0} ({1})"), FText::FromString(WatchdogCommandLine.ProjectName), FText::FromString(WatchdogCommandLine.RunType));
+		FString EnsureText;
+
+		if (StoredValues.ExecutionStatus == WatchdogDefs::CrashSessionToken)
+		{
+			// Crashed status
+			FText MessageTitleFormat(LOCTEXT("WatchdogPopupTitleCrashed", "{0} crashed"));
+			FText MessageTitle = FText::Format(MessageTitleFormat, SessionLabel);
+			FText Message(LOCTEXT("WatchdogPopupQuestionCrashed", "We think a crash occurred, was handled correctly, and the Crash Reporter appeared. Please tell us if you saw the Crash Reporter?"));
+
+			EnsureText = FString::Printf(TEXT("Watchdog detected crash in %s."), *SessionLabel.ToString());
+			UserAnswer = FMessageDialog::Open(EAppMsgType::YesNo, Message, &MessageTitle);
+
+			if (UserAnswer == EAppReturnType::Yes)
+			{
+				UE_LOG(UnrealWatchdogLog, Log, TEXT("User confirmed crash and crash report client"));
+				UserResponse.AttrValue = TEXT("Confirmed");
+			}
+			else if (UserAnswer == EAppReturnType::No)
+			{
+				UE_LOG(UnrealWatchdogLog, Warning, TEXT("User didn't witness crash and crash report client. False positive warning!"));
+				UserResponse.AttrValue = TEXT("False");
+
+				FText CRCMessage(LOCTEXT("WatchdogPopupCRCMessage", "We will now open the Crash Reporter for you to tell us what happened."));
+				FMessageDialog::Open(EAppMsgType::Ok, CRCMessage, &MessageTitle);
+
+				ReportInteractiveEnsure(*EnsureText);
+			}
+		}
+		else
+		{
+			// Non-crash status
+			const bool bBadReturnCode = bReturnCodeObtained && ReturnCode != WatchdogCommandLine.SuccessReturnCode;
+
+			if (StoredValues.ExecutionStatus == WatchdogDefs::RunningSessionToken)
+			{
+				// Abnormal shutdown status
+				FText MessageTitle;
+				FText Message;
+				if (bBadReturnCode)
+				{
+					FText MessageTitleFormat(LOCTEXT("WatchdogPopupTitleAbnormalShutdown", "{0} terminated unexpectedly"));
+					MessageTitle = FText::Format(MessageTitleFormat, SessionLabel);
+					Message = LOCTEXT("WatchdogPopupQuestionAbnormalShutdown", "We think a crash was not handled correctly. Did the application crash without showing the Crash Reporter?");
+
+					EnsureText = FString::Printf(TEXT("Watchdog detected abnormal shutdown and returned error code %d in %s."), ReturnCode, *SessionLabel.ToString());
+				}
+				else
+				{
+					FText MessageTitleFormat(LOCTEXT("WatchdogPopupTitlePossibleAbnormalShutdown", "{0} did not close correctly"));
+					MessageTitle = FText::Format(MessageTitleFormat, SessionLabel);
+					Message = LOCTEXT("WatchdogPopupQuestionPossibleAbnormalShutdown", "We think we exited normally but didn't cleanup correctly. Did the application crash without showing the Crash Reporter?");
+
+					EnsureText = FString::Printf(TEXT("Watchdog detected abnormal shutdown with successful return code %d in %s."), ReturnCode, *SessionLabel.ToString());
+				}
+
+				UserAnswer = FMessageDialog::Open(EAppMsgType::YesNo, Message, &MessageTitle);
+			}
+			else if (StoredValues.ExecutionStatus == WatchdogDefs::ShutdownSessionToken)
+			{
+				if (bBadReturnCode)
+				{
+					// Normal shutdown status with bad return code
+					FText MessageTitleFormat(LOCTEXT("WatchdogPopupTitleNormalShutdown", "{0} returned an error code"));
+					FText MessageTitle = FText::Format(MessageTitleFormat, SessionLabel);
+					FText Message(LOCTEXT("WatchdogPopupQuestionNormalShutdown", "We think we shutdown correctly but returned an error code. Did the application crash without showing the Crash Reporter?"));
+
+					EnsureText = FString::Printf(TEXT("Watchdog detected normal shutdown but returned error code %d in %s."), ReturnCode, *SessionLabel.ToString());
+					UserAnswer = FMessageDialog::Open(EAppMsgType::YesNo, Message, &MessageTitle);
+				}
+			}
+			else if (bBadReturnCode)
+			{
+				// Missing or unknown status with error code
+				FText MessageTitleFormat(LOCTEXT("WatchdogPopupTitleUnknownError", "{0} returned an error code"));
+				FText MessageTitle = FText::Format(MessageTitleFormat, SessionLabel);
+				FText Message(LOCTEXT("WatchdogPopupQuestionUnknownError", "Did the application crash without showing the Crash Reporter?"));
+
+				EnsureText = FString::Printf(TEXT("Watchdog detected unknown shutdown status and returned error code %d in %s."), ReturnCode, *SessionLabel.ToString());
+				UserAnswer = FMessageDialog::Open(EAppMsgType::YesNo, Message, &MessageTitle);
+			}
+			else
+			{
+				// Missing or unknown status (no error code)
+				FText MessageTitleFormat(LOCTEXT("WatchdogPopupTitleUnknownNoError", "{0} failed to signal the watchdog"));
+				FText MessageTitle = FText::Format(MessageTitleFormat, SessionLabel);
+				FText Message(LOCTEXT("WatchdogPopupQuestionUnknownNoError", "Process didn't communicate with this watchdog correctly. Did the application crash without showing the Crash Reporter?"));
+
+				EnsureText = FString::Printf(TEXT("Watchdog detected unknown shutdown status with successful return code %d in %s."), ReturnCode, *SessionLabel.ToString());
+				UserAnswer = FMessageDialog::Open(EAppMsgType::YesNo, Message, &MessageTitle);
+			}
+
+			if (UserAnswer == EAppReturnType::Yes)
+			{
+				UE_LOG(UnrealWatchdogLog, Log, TEXT("User confirmed abnormal shutdown"));
+				UserResponse.AttrValue = TEXT("Confirmed");
+
+				FText MessageTitleFormat(LOCTEXT("WatchdogPopupTitleAbnormalShutdown", "{0} terminated unexpectedly"));
+				FText MessageTitle = FText::Format(MessageTitleFormat, SessionLabel);
+				FText Message(LOCTEXT("WatchdogPopupCRCMessage", "We will now open the Crash Reporter for you to tell us what happened."));
+
+				FMessageDialog::Open(EAppMsgType::Ok, Message, &MessageTitle);
+
+				ReportInteractiveEnsure(*EnsureText);
+			}
+			else if (UserAnswer == EAppReturnType::No)
+			{
+				UE_LOG(UnrealWatchdogLog, Warning, TEXT("User didn't witness abnormal shutdown. False positive warning!"));
+				UserResponse.AttrValue = TEXT("False");
+			}
+		}
+	}
+
+	// Send watchdog shutdown event
+	UE_LOG(UnrealWatchdogLog, Log, TEXT("Watchdog watched process exited. bReturnCodeObtained=%s, ReturnCode=%u, RecordedShutdownType=%s"),
+		bReturnCodeObtained ? TEXT("1") : TEXT("0"), ReturnCode, *StoredValues.ExecutionStatus);
+	SendShutdownEvent(Analytics, WatchdogCommandLine, bReturnCodeObtained, ReturnCode, UserResponse, StoredValues);
+
+	// Shutdown tool and engine
+	FWatchdogAnalytics::Shutdown();
+	UE_LOG(UnrealWatchdogLog, Log, TEXT("Watchdog exiting"));
+	FEngineLoop::AppPreExit();
+	FModuleManager::Get().UnloadModulesAtShutdown();
+
+	FEngineLoop::AppExit();
+	
+	return 0;
+}
+
+#undef LOCTEXT_NAMESPACE
diff --git a/Engine/Source/Programs/UnrealWatchdog/Private/UnrealWatchdog.h b/Engine/Source/Programs/UnrealWatchdog/Private/UnrealWatchdog.h
new file mode 100644
--- /dev/null
+++ b/Engine/Source/Programs/UnrealWatchdog/Private/UnrealWatchdog.h
@@ -0,0 +1,42 @@
+// Copyright 1998-2016 Epic Games, Inc. All Rights Reserved.
+
+#pragma once
+
+DECLARE_LOG_CATEGORY_EXTERN(UnrealWatchdogLog, Log, All);
+
+struct FWatchdogCommandLine
+{
+	FString RunType;
+	FString ProjectName;
+	FString PlatformName;
+	FString SessionId;
+	FString EngineVersion;
+	uint32 ParentProcessId;
+	int32 SuccessReturnCode;
+	int32 HeartbeatSeconds;
+	bool bAllowDialogs;
+	bool bHasProcessId;
+
+	FWatchdogCommandLine(const TCHAR* InCommandLine)
+	{
+		// Read command line args
+		FParse::Value(InCommandLine, TEXT("RunType="), RunType);
+		FParse::Value(InCommandLine, TEXT("ProjectName="), ProjectName);
+		FParse::Value(InCommandLine, TEXT("Platform="), PlatformName);
+		FParse::Value(InCommandLine, TEXT("SessionId="), SessionId);
+		FParse::Value(InCommandLine, TEXT("EngineVersion="), EngineVersion);
+		bHasProcessId = FParse::Value(InCommandLine, TEXT("PID="), ParentProcessId);
+		SuccessReturnCode = 0;
+		FParse::Value(InCommandLine, TEXT("SuccessfulRtnCode="), SuccessReturnCode);
+		HeartbeatSeconds = 300;
+		FParse::Value(InCommandLine, TEXT("HeartbeatSeconds="), HeartbeatSeconds);
+		bAllowDialogs = FParse::Param(InCommandLine, TEXT("AllowDialogs"));
+	}
+};
+
+int RunUnrealWatchdog(const TCHAR* CommandLine);
+bool WaitForProcess(class IAnalyticsProviderET& Analytics, const FWatchdogCommandLine& CommandLine, int32& OutReturnCode, bool& bOutHang, const FString& WatchdogSectionName);
+void SendHeartbeatEvent(class IAnalyticsProviderET& Analytics, const FWatchdogCommandLine& CommandLine);
+bool CheckParentHeartbeat(class IAnalyticsProviderET& Analytics, const FWatchdogCommandLine& CommandLine, const FString& WatchdogSectionName);
+void SendHangDetectedEvent(class IAnalyticsProviderET& Analytics, const FWatchdogCommandLine& CommandLine);
+void SendHangRecoveredEvent(class IAnalyticsProviderET& Analytics, const FWatchdogCommandLine& CommandLine);
diff --git a/Engine/Source/Programs/UnrealWatchdog/Private/WatchdogAnalytics.cpp b/Engine/Source/Programs/UnrealWatchdog/Private/WatchdogAnalytics.cpp
new file mode 100644
--- /dev/null
+++ b/Engine/Source/Programs/UnrealWatchdog/Private/WatchdogAnalytics.cpp
@@ -0,0 +1,63 @@
+// Copyright 1998-2016 Epic Games, Inc. All Rights Reserved.
+
+#include "WatchdogAnalytics.h"
+#include "UnrealWatchdog.h"
+#include "EngineBuildSettings.h"
+#include "AnalyticsET.h"
+#include "IAnalyticsProviderET.h"
+
+bool FWatchdogAnalytics::bIsInitialized;
+TSharedPtr<IAnalyticsProviderET> FWatchdogAnalytics::Analytics;
+
+/**
+* Default config func that essentially tells the crash reporter to disable analytics.
+*/
+FAnalyticsET::Config DefaultAnalyticsConfigFunc()
+{
+	return FAnalyticsET::Config();
+}
+
+/**
+* Engine analytics config to initialize the watchdog analytics provider.
+* External code should bind this delegate if watchdog analytics in private code that won't be redistributed.
+*/
+TFunction<FAnalyticsET::Config()>& GetWatchdogAnalyticsConfigFunc()
+{
+	static TFunction<FAnalyticsET::Config()> Config = &DefaultAnalyticsConfigFunc;
+	return Config;
+}
+
+/**
+* On-demand construction of the singleton.
+*/
+IAnalyticsProviderET& FWatchdogAnalytics::GetProvider()
+{
+	checkf(bIsInitialized && Analytics.IsValid(), TEXT("FWatchdogAnalytics::GetProvider called outside of Initialize/Shutdown."));
+	return *Analytics.Get();
+}
+
+void FWatchdogAnalytics::Initialize()
+{
+	checkf(!bIsInitialized, TEXT("FWatchdogAnalytics::Initialize called more than once."));
+
+	FAnalyticsET::Config Config = GetWatchdogAnalyticsConfigFunc()();
+	if (!Config.APIServerET.IsEmpty())
+	{
+		// Connect the engine analytics provider (if there is a configuration delegate installed)
+		Analytics = FAnalyticsET::Get().CreateAnalyticsProvider(Config);
+		if (Analytics.IsValid())
+		{
+			Analytics->SetUserID(FString::Printf(TEXT("%s|%s|%s"), *FPlatformMisc::GetLoginId(), *FPlatformMisc::GetEpicAccountId(), *FPlatformMisc::GetOperatingSystemId()));
+			Analytics->StartSession();
+		}
+	}
+	bIsInitialized = true;
+}
+
+
+void FWatchdogAnalytics::Shutdown()
+{
+	checkf(bIsInitialized, TEXT("FWatchdogAnalytics::Shutdown called outside of Initialize."));
+	Analytics.Reset();
+	bIsInitialized = false;
+}
diff --git a/Engine/Source/Programs/UnrealWatchdog/Private/WatchdogAnalytics.h b/Engine/Source/Programs/UnrealWatchdog/Private/WatchdogAnalytics.h
new file mode 100644
--- /dev/null
+++ b/Engine/Source/Programs/UnrealWatchdog/Private/WatchdogAnalytics.h
@@ -0,0 +1,36 @@
+// Copyright 1998-2016 Epic Games, Inc. All Rights Reserved.
+
+#pragma once
+
+class IAnalyticsProviderET;
+
+/**
+* The public interface for the analytics provider singleton.
+* For Epic builds, this will point to epic's internal analytics provider.
+* For licensee builds, it will be NULL by default unless they provide their own
+* configuration.
+*
+*/
+class FWatchdogAnalytics : FNoncopyable
+{
+public:
+	/**
+	* Return the provider instance. Not valid outside of Initialize/Shutdown calls.
+	* Note: must check IsAvailable() first else this code will assert if the provider is not valid.
+	*/
+	static IAnalyticsProviderET& GetProvider();
+	/** Helper function to determine if the provider is valid. */
+	static bool IsAvailable()
+	{
+		return Analytics.IsValid();
+	}
+	/** Called to initialize the singleton. */
+	static void Initialize();
+	/** Called to shut down the singleton */
+	static void Shutdown();
+
+private:
+	static bool bIsInitialized;
+	static TSharedPtr<IAnalyticsProviderET> Analytics;
+};
+
diff --git a/Engine/Source/Programs/UnrealWatchdog/Private/Windows/UnrealWatchdogMainWindows.cpp b/Engine/Source/Programs/UnrealWatchdog/Private/Windows/UnrealWatchdogMainWindows.cpp
new file mode 100644
--- /dev/null
+++ b/Engine/Source/Programs/UnrealWatchdog/Private/Windows/UnrealWatchdogMainWindows.cpp
@@ -0,0 +1,111 @@
+// Copyright 1998-2016 Epic Games, Inc. All Rights Reserved.
+
+#include "UnrealWatchdog.h"
+#include "WatchdogAnalytics.h"
+
+/**
+* WinMain, called when the application is started
+*/
+INT32_MAIN_INT32_ARGC_TCHAR_ARGV()
+{
+	FString CommandLine;
+	for (int32 Option = 1; Option < ArgC; Option++)
+	{
+		CommandLine += TEXT(" ");
+		FString Argument(ArgV[Option]);
+		if (Argument.Contains(TEXT(" ")))
+		{
+			if (Argument.Contains(TEXT("=")))
+			{
+				FString ArgName;
+				FString ArgValue;
+				Argument.Split(TEXT("="), &ArgName, &ArgValue);
+				Argument = FString::Printf(TEXT("%s=\"%s\""), *ArgName, *ArgValue);
+			}
+			else
+			{
+				Argument = FString::Printf(TEXT("\"%s\""), *Argument);
+			}
+		}
+		CommandLine += Argument;
+	}
+
+	// hide the console window
+	ShowWindow(GetConsoleWindow(), SW_HIDE);
+
+	// Run the app
+	RunUnrealWatchdog(*CommandLine);
+
+	return 0;
+}
+
+bool WaitForProcess(IAnalyticsProviderET& Analytics, const FWatchdogCommandLine& CommandLine, int32& OutReturnCode, bool& bOutHang, const FString& WatchdogSectionName)
+{
+	const int HangDetectionGraceSeconds = 10;
+	const FTimespan CheckParentHeartbeatPeriod(0, 0, CommandLine.HeartbeatSeconds + HangDetectionGraceSeconds);
+	const FTimespan SendWatchdogHeartbeatPeriod(0, 5, 0);
+	const FTimespan CheckParentRunningPeriod(0, 0, 10);
+
+	FDateTime NextHeartbeatSend = FDateTime::UtcNow();
+	FDateTime NextHeartbeatCheck = FDateTime::UtcNow();
+	bool bSuccess = false;
+	bOutHang = false;
+	OutReturnCode = -1;
+
+	FProcHandle ParentProcess = FPlatformProcess::OpenProcess(CommandLine.ParentProcessId);
+	if (ParentProcess.IsValid())
+	{
+		while (!GIsRequestingExit)
+		{
+			// Send heartbeat due?
+			if (FDateTime::UtcNow() >= NextHeartbeatSend)
+			{
+				NextHeartbeatSend += SendWatchdogHeartbeatPeriod;
+				SendHeartbeatEvent(Analytics, CommandLine);
+			}
+
+			// Parent application still running?
+			if (!FPlatformProcess::IsApplicationRunning(CommandLine.ParentProcessId))
+			{
+				UE_LOG(UnrealWatchdogLog, Log, TEXT("Watchdog detected terminated process PID %u"), CommandLine.ParentProcessId);
+				bSuccess = FPlatformProcess::GetProcReturnCode(ParentProcess, &OutReturnCode);
+				break;
+			}
+
+			// Check parent heartbeat for hang?
+			if (FDateTime::UtcNow() >= NextHeartbeatCheck)
+			{
+				NextHeartbeatCheck += CheckParentHeartbeatPeriod;
+				bool bHangDetected = !CheckParentHeartbeat(Analytics, CommandLine, WatchdogSectionName);
+
+				if (bHangDetected && !bOutHang)
+				{
+					// New hang
+					SendHangDetectedEvent(Analytics, CommandLine);
+					bOutHang = true;
+				}
+				else if (!bHangDetected && bOutHang)
+				{
+					// Previous hang recovered
+					SendHangRecoveredEvent(Analytics, CommandLine);
+					bOutHang = false;
+				}
+			}
+
+			// Sleep until next event
+			FDateTime EndLoopTime = FDateTime::UtcNow();
+			FTimespan TimeToHeartbeatCheck = NextHeartbeatCheck - EndLoopTime;
+			FTimespan TimeToHeartbeatSend = NextHeartbeatSend - EndLoopTime;
+			double SleepSeconds = FMath::Min3(TimeToHeartbeatCheck.GetTotalSeconds(), TimeToHeartbeatSend.GetTotalSeconds(), CheckParentRunningPeriod.GetTotalSeconds());
+			FPlatformProcess::Sleep((float)SleepSeconds);
+		}
+
+		FPlatformProcess::CloseProc(ParentProcess);
+	}
+	else
+	{
+		UE_LOG(UnrealWatchdogLog, Error, TEXT("Watchdog failed to get handle to process PID %u"), CommandLine.ParentProcessId);
+	}
+
+	return bSuccess;
+}
\ No newline at end of file
diff --git a/Engine/Source/Programs/UnrealWatchdog/UnrealWatchdog.Build.cs b/Engine/Source/Programs/UnrealWatchdog/UnrealWatchdog.Build.cs
new file mode 100644
--- /dev/null
+++ b/Engine/Source/Programs/UnrealWatchdog/UnrealWatchdog.Build.cs
@@ -0,0 +1,30 @@
+﻿// Copyright 1998-2016 Epic Games, Inc. All Rights Reserved.
+
+using UnrealBuildTool;
+
+public class UnrealWatchdog : ModuleRules
+{
+	public UnrealWatchdog(TargetInfo Target)
+	{
+		PublicIncludePaths.AddRange
+		(
+			new string[]
+			{
+				"Runtime/Launch/Public",
+				"Programs/UnrealWatchdog/Private",
+			}
+		);
+
+		// For LaunchEngineLoop.cpp include
+		PrivateIncludePaths.Add("Runtime/Launch/Private");
+
+		PrivateDependencyModuleNames.AddRange(
+			new string[] {
+                "Core",
+				"Analytics",
+				"AnalyticsET",
+				"Projects",
+		   }
+		);
+	}
+}
diff --git a/Engine/Source/Programs/UnrealWatchdog/UnrealWatchdog.Target.cs b/Engine/Source/Programs/UnrealWatchdog/UnrealWatchdog.Target.cs
new file mode 100644
--- /dev/null
+++ b/Engine/Source/Programs/UnrealWatchdog/UnrealWatchdog.Target.cs
@@ -0,0 +1,85 @@
+﻿// Copyright 1998-2016 Epic Games, Inc. All Rights Reserved.
+
+using UnrealBuildTool;
+using System.Collections.Generic;
+
+public class UnrealWatchdogTarget : TargetRules
+{
+	public UnrealWatchdogTarget(TargetInfo Target)
+	{
+		Type = TargetType.Program;
+        UndecoratedConfiguration = UnrealTargetConfiguration.Shipping;
+    }
+
+    //
+    // TargetRules interface.
+    //
+
+    public override bool GetSupportedPlatforms(ref List<UnrealTargetPlatform> OutPlatforms)
+	{
+		OutPlatforms.Add(UnrealTargetPlatform.Win32);
+		OutPlatforms.Add(UnrealTargetPlatform.Win64);
+		//OutPlatforms.Add(UnrealTargetPlatform.Mac);
+		//OutPlatforms.Add(UnrealTargetPlatform.Linux);
+		return true;
+	}
+
+    public override bool GetSupportedConfigurations(ref List<UnrealTargetConfiguration> OutConfigurations, bool bIncludeTestAndShippingConfigs)
+    {
+        if (base.GetSupportedConfigurations(ref OutConfigurations, bIncludeTestAndShippingConfigs))
+        {
+            if (bIncludeTestAndShippingConfigs)
+            {
+                OutConfigurations.Add(UnrealTargetConfiguration.Shipping);
+            }
+            OutConfigurations.Add(UnrealTargetConfiguration.Debug);
+            return true;
+        }
+        else
+        {
+            return false;
+        }
+    }
+
+    public override void SetupBinaries(
+		TargetInfo Target,
+		ref List<UEBuildBinaryConfiguration> OutBuildBinaryConfigurations,
+		ref List<string> OutExtraModuleNames
+		)
+	{
+		OutBuildBinaryConfigurations.Add(
+			new UEBuildBinaryConfiguration(	InType: UEBuildBinaryType.Executable,
+											InModuleNames: new List<string>() { "UnrealWatchdog" } )
+			);
+	}
+
+	public override bool ShouldCompileMonolithic(UnrealTargetPlatform InPlatform, UnrealTargetConfiguration InConfiguration)
+	{
+		return true;
+	}
+
+	public override void SetupGlobalEnvironment(
+		TargetInfo Target,
+		ref LinkEnvironmentConfiguration OutLinkEnvironmentConfiguration,
+		ref CPPEnvironmentConfiguration OutCPPEnvironmentConfiguration
+		)
+	{
+		// Lean and mean
+		UEBuildConfiguration.bCompileLeanAndMeanUE = true;
+		BuildConfiguration.bUseMallocProfiler = false;
+
+        UEBuildConfiguration.bUseLoggingInShipping = true;
+
+        // No editor needed
+        UEBuildConfiguration.bBuildEditor = false;
+		UEBuildConfiguration.bBuildWithEditorOnlyData = false;
+
+		// Currently this app is not linking against the engine, so we'll compile out references from Core to the rest of the engine
+		UEBuildConfiguration.bCompileAgainstEngine = false;
+		UEBuildConfiguration.bCompileAgainstCoreUObject = false;
+		UEBuildConfiguration.bBuildDeveloperTools = false;
+
+		// Console application, not a Windows app (sets entry point to main(), instead of WinMain())
+		OutLinkEnvironmentConfiguration.bIsBuildingConsoleApplication = true;
+	}
+}
diff --git a/Engine/Source/Runtime/Core/Private/GenericPlatform/GenericPlatformProcess.cpp b/Engine/Source/Runtime/Core/Private/GenericPlatform/GenericPlatformProcess.cpp
--- a/Engine/Source/Runtime/Core/Private/GenericPlatform/GenericPlatformProcess.cpp
+++ b/Engine/Source/Runtime/Core/Private/GenericPlatform/GenericPlatformProcess.cpp
@@ -206,6 +206,12 @@ FProcHandle FGenericPlatformProcess::CreateProc( const TCHAR* URL, const TCHAR*
 	return FProcHandle();
 }
 
+FProcHandle FGenericPlatformProcess::OpenProcess(uint32 ProcessID)
+{
+	UE_LOG(LogHAL, Fatal, TEXT("FGenericPlatformProcess::OpenProcess not implemented on this platform"));
+	return FProcHandle();
+}
+
 bool FGenericPlatformProcess::IsProcRunning( FProcHandle & ProcessHandle )
 {
 	UE_LOG(LogHAL, Fatal, TEXT("FGenericPlatformProcess::IsProcRunning not implemented on this platform"));
diff --git a/Engine/Source/Runtime/Core/Private/HAL/ExceptionHandling.cpp b/Engine/Source/Runtime/Core/Private/HAL/ExceptionHandling.cpp
--- a/Engine/Source/Runtime/Core/Private/HAL/ExceptionHandling.cpp
+++ b/Engine/Source/Runtime/Core/Private/HAL/ExceptionHandling.cpp
@@ -29,6 +29,7 @@ CORE_API TCHAR MiniDumpFilenameW[1024] = TEXT("");
 
 
 volatile int32 GCrashType = 0;
+bool GEnsureShowsCRC = false;
 
 void ReportImageIntegrityStatus(const TCHAR* InMessage, const int32 InCode)
 {
@@ -87,3 +88,23 @@ int32 GetCrashType()
 {
 	return GCrashType;
 }
+void ReportInteractiveEnsure(const TCHAR* InMessage)
+{
+	GEnsureShowsCRC = true;
+
+#if PLATFORM_DESKTOP
+	GLog->PanicFlushThreadedLogs();
+	// GErrorMessage here is very unfortunate but it's used internally by the crash context code.
+	FCString::Strcpy(GErrorMessage, ARRAY_COUNT(GErrorMessage), InMessage);
+	// Skip macros and FDebug, we always want this to fire
+	NewReportEnsure(InMessage);
+	GErrorMessage[0] = '\0';
+#endif
+
+	GEnsureShowsCRC = false;
+}
+
+bool IsInteractiveEnsureMode()
+{
+	return GEnsureShowsCRC;
+}
\ No newline at end of file
diff --git a/Engine/Source/Runtime/Core/Private/Linux/LinuxPlatformCrashContext.cpp b/Engine/Source/Runtime/Core/Private/Linux/LinuxPlatformCrashContext.cpp
--- a/Engine/Source/Runtime/Core/Private/Linux/LinuxPlatformCrashContext.cpp
+++ b/Engine/Source/Runtime/Core/Private/Linux/LinuxPlatformCrashContext.cpp
@@ -22,6 +22,7 @@
 #include "Misc/EngineVersion.h"
 #include "HAL/PlatformMallocCrash.h"
 #include "Linux/LinuxPlatformRunnableThread.h"
+#include "ExceptionHandling.h"
 
 #include "HAL/ThreadHeartBeat.h"
 
@@ -475,7 +476,7 @@ void FLinuxCrashContext::GenerateCrashInfoAndLaunchReporter(bool bReportingNonCr
 		CrashReportClientArguments += TEXT(" ");
 
 		// Suppress the user input dialog if we're running in unattended mode
-		bool bNoDialog = FApp::IsUnattended() || IsRunningDedicatedServer();
+		bool bNoDialog = FApp::IsUnattended() || (!IsInteractiveEnsureMode() && bReportingNonCrash) || IsRunningDedicatedServer();
 		if (bNoDialog)
 		{
 			CrashReportClientArguments += TEXT(" -Unattended ");
diff --git a/Engine/Source/Runtime/Core/Private/Mac/MacPlatformMisc.cpp b/Engine/Source/Runtime/Core/Private/Mac/MacPlatformMisc.cpp
--- a/Engine/Source/Runtime/Core/Private/Mac/MacPlatformMisc.cpp
+++ b/Engine/Source/Runtime/Core/Private/Mac/MacPlatformMisc.cpp
@@ -2389,7 +2389,16 @@ void FMacCrashContext::GenerateEnsureInfoAndLaunchReporter() const
 		const bool bIsEnsure = true;
 		GenerateInfoInFolder(TCHAR_TO_UTF8(*EnsureLogFolder), bIsEnsure);
 		
-		FString Arguments = FString::Printf(TEXT("\"%s/\" -Unattended"), *EnsureLogFolder);
+		FString Arguments;
+		if (IsInteractiveEnsureMode())
+		{
+			Arguments = FString::Printf(TEXT("\"%s/\""), *EnsureLogFolder);
+		}
+		else
+		{
+			Arguments = FString::Printf(TEXT("\"%s/\" -Unattended"), *EnsureLogFolder);
+		}
+
 		FString ReportClient = FPaths::ConvertRelativePathToFull(FPlatformProcess::GenerateApplicationPath(TEXT("CrashReportClient"), EBuildConfigurations::Development));
 		FPlatformProcess::ExecProcess(*ReportClient, *Arguments, nullptr, nullptr, nullptr);
 	}
diff --git a/Engine/Source/Runtime/Core/Private/Mac/MacPlatformProcess.cpp b/Engine/Source/Runtime/Core/Private/Mac/MacPlatformProcess.cpp
--- a/Engine/Source/Runtime/Core/Private/Mac/MacPlatformProcess.cpp
+++ b/Engine/Source/Runtime/Core/Private/Mac/MacPlatformProcess.cpp
@@ -82,6 +82,7 @@ FString FMacPlatformProcess::GenerateApplicationPath( const FString& AppName, EB
 	}
 	else
 	{
+		// Try expected path of an executable inside an app package in Engine Binaries
 		FString ExecutablePath = FString::Printf(TEXT("../../../Engine/Binaries/%s/%s.app/Contents/MacOS/%s"), *PlatformName, *ExecutableName, *ExecutableName);
 			
 		NSString* LaunchPath = ExecutablePath.GetNSString();
@@ -92,18 +93,30 @@ FString FMacPlatformProcess::GenerateApplicationPath( const FString& AppName, EB
 		}
 		else
 		{
-			CFStringRef App = FPlatformString::TCHARToCFString(*ExecutableName);
-			NSWorkspace* Workspace = [NSWorkspace sharedWorkspace];
-			NSString* AppPath = [Workspace fullPathForApplication:(NSString*)App];
-			CFRelease(App);
-			if (AppPath)
+			// Next try expected path of a simple executable file in Engine Binaries
+			ExecutablePath = FString::Printf(TEXT("../../../Engine/Binaries/%s/%s"), *PlatformName, *ExecutableName);
+
+			LaunchPath = ExecutablePath.GetNSString();
+
+			if ([[NSFileManager defaultManager] fileExistsAtPath:LaunchPath])
 			{
-				ExecutablePath = FString::Printf(TEXT("%s/Contents/MacOS/%s"), *FString(AppPath), *ExecutableName);
 				return ExecutablePath;
 			}
 			else
 			{
-				return FString();
+				CFStringRef App = FPlatformString::TCHARToCFString(*ExecutableName);
+				NSWorkspace* Workspace = [NSWorkspace sharedWorkspace];
+				NSString* AppPath = [Workspace fullPathForApplication : (NSString*)App];
+				CFRelease(App);
+				if (AppPath)
+				{
+					ExecutablePath = FString::Printf(TEXT("%s/Contents/MacOS/%s"), *FString(AppPath), *ExecutableName);
+					return ExecutablePath;
+				}
+				else
+				{
+					return FString();
+				}
 			}
 		}
 	}
diff --git a/Engine/Source/Runtime/Core/Private/Windows/WindowsPlatformCrashContext.cpp b/Engine/Source/Runtime/Core/Private/Windows/WindowsPlatformCrashContext.cpp
--- a/Engine/Source/Runtime/Core/Private/Windows/WindowsPlatformCrashContext.cpp
+++ b/Engine/Source/Runtime/Core/Private/Windows/WindowsPlatformCrashContext.cpp
@@ -307,7 +307,7 @@ void NewReportEnsure( const TCHAR* ErrorMessage )
 		FPlatformMisc::RaiseException( 1 );
 	}
 #if !PLATFORM_SEH_EXCEPTIONS_DISABLED
-	__except(ReportEnsureUsingCrashReportClient(GetExceptionInformation(), ErrorMessage, EErrorReportUI::ReportInUnattendedMode))
+	__except(ReportEnsureUsingCrashReportClient(GetExceptionInformation(), ErrorMessage, IsInteractiveEnsureMode() ? EErrorReportUI::ShowDialog : EErrorReportUI::ReportInUnattendedMode))
 	CA_SUPPRESS(6322)
 	{
 	}
diff --git a/Engine/Source/Runtime/Core/Public/GenericPlatform/GenericPlatformProcess.h b/Engine/Source/Runtime/Core/Public/GenericPlatform/GenericPlatformProcess.h
--- a/Engine/Source/Runtime/Core/Public/GenericPlatform/GenericPlatformProcess.h
+++ b/Engine/Source/Runtime/Core/Public/GenericPlatform/GenericPlatformProcess.h
@@ -344,6 +344,14 @@ struct CORE_API FGenericPlatformProcess
 	 */
 	static FProcHandle CreateProc( const TCHAR* URL, const TCHAR* Parms, bool bLaunchDetached, bool bLaunchHidden, bool bLaunchReallyHidden, uint32* OutProcessID, int32 PriorityModifier, const TCHAR* OptionalWorkingDirectory, void* PipeWriteChild, void * PipeReadChild = nullptr);
 
+	/**
+	 * Opens an existing process. 
+	 *
+	 * @param ProcessID				The process id of the process for which we want to obtain a handle.
+	 * @return The process handle for use in other process functions
+	 */
+	static FProcHandle OpenProcess(uint32 ProcessID);
+
 	/**
 	 * Returns true if the specified process is running 
 	 *
diff --git a/Engine/Source/Runtime/Core/Public/HAL/ExceptionHandling.h b/Engine/Source/Runtime/Core/Public/HAL/ExceptionHandling.h
--- a/Engine/Source/Runtime/Core/Public/HAL/ExceptionHandling.h
+++ b/Engine/Source/Runtime/Core/Public/HAL/ExceptionHandling.h
@@ -48,3 +48,5 @@ enum class ECrashType
 /** In debug builds we can crash on demand so we want to let the crash reporting system know about it */
 extern CORE_API void SetCrashType(ECrashType InCrashType);
 extern CORE_API int32 GetCrashType();
+extern CORE_API void ReportInteractiveEnsure(const TCHAR* InMessage);
+extern CORE_API bool IsInteractiveEnsureMode();
\ No newline at end of file
diff --git a/Engine/Source/Runtime/Engine/Private/EngineSessionManager.cpp b/Engine/Source/Runtime/Engine/Private/EngineSessionManager.cpp
--- a/Engine/Source/Runtime/Engine/Private/EngineSessionManager.cpp
+++ b/Engine/Source/Runtime/Engine/Private/EngineSessionManager.cpp
@@ -11,17 +11,24 @@
 #include "UserActivityTracking.h"
 #include "Engine/Engine.h"
 #include "EngineGlobals.h"
+#include "EngineBuildSettings.h"
+#include "SlateApplication.h"
+
 
 #define LOCTEXT_NAMESPACE "SessionManager"
 
+DEFINE_LOG_CATEGORY(LogEngineSessionManager);
+
 namespace SessionManagerDefs
 {
 	static const FTimespan SessionRecordExpiration(30, 0, 0, 0);	// 30 days
-	static const FTimespan SessionRecordTimeout(0, 30, 0);			// 30 minutes
+	static const FTimespan SessionRecordTimeout(0, 3, 0);			// 3 minutes
 	static const FTimespan GlobalLockWaitTimeout(0, 0, 0, 0, 500);	// 1/2 second
-	static const float HeartbeatPeriodSeconds(300.0f);				// 5 minutes
+	static const int HeartbeatPeriodSeconds(60);				// 1 minute
 	static const FString DefaultUserActivity(TEXT("Unknown"));
 	static const FString StoreId(TEXT("Epic Games"));
+	static const FString RunningSessionToken(TEXT("Running"));
+	static const FString ShutdownSessionToken(TEXT("Shutdown"));
 	static const FString CrashSessionToken(TEXT("Crashed"));
 	static const FString DebuggerSessionToken(TEXT("Debugger"));
 	static const FString AbnormalSessionToken(TEXT("AbnormalShutdown"));
@@ -29,14 +36,20 @@ namespace SessionManagerDefs
 	static const FString SessionRecordListSection(TEXT("List"));
 	static const FString EditorSessionRecordSectionPrefix(TEXT("Unreal Engine/Editor Sessions/"));
 	static const FString GameSessionRecordSectionPrefix(TEXT("Unreal Engine/Game Sessions/"));
+	static const FString WatchdogRecordSectionPrefix(TEXT("Unreal Engine/Watchdog/"));
 	static const FString SessionsVersionString(TEXT("1_3"));
+	static const FString WatchdogVersionString(TEXT("1_0"));
 	static const FString ModeStoreKey(TEXT("Mode"));
 	static const FString ProjectNameStoreKey(TEXT("ProjectName"));
+	static const FString CommandLineStoreKey(TEXT("CommandLine"));
 	static const FString CrashStoreKey(TEXT("IsCrash"));
 	static const FString DeactivatedStoreKey(TEXT("IsDeactivated"));
 	static const FString BackgroundStoreKey(TEXT("IsInBackground"));
 	static const FString EngineVersionStoreKey(TEXT("EngineVersion"));
 	static const FString TimestampStoreKey(TEXT("Timestamp"));
+	static const FString StartupTimeStoreKey(TEXT("StartupTimestamp"));
+	static const FString SessionIdStoreKey(TEXT("SessionId"));
+	static const FString StatusStoreKey(TEXT("LastExecutionState"));
 	static const FString DebuggerStoreKey(TEXT("IsDebugger"));
 	static const FString UserActivityStoreKey(TEXT("CurrentUserActivity"));
 	static const FString VanillaStoreKey(TEXT("IsVanilla"));
@@ -78,7 +91,8 @@ void FEngineSessionManager::Initialize()
 	FCoreDelegates::ApplicationHasEnteredForegroundDelegate.AddRaw(this, &FEngineSessionManager::OnAppForeground);
 	FUserActivityTracking::OnActivityChanged.AddRaw(this, &FEngineSessionManager::OnUserActivity);
 	FCoreDelegates::IsVanillaProductChanged.AddRaw(this, &FEngineSessionManager::OnVanillaStateChanged);
-
+	FSlateApplication::Get().GetOnModalLoopTickEvent().AddRaw(this, &FEngineSessionManager::Tick);
+	
 	const bool bFirstInitAttempt = true;
 	InitializeRecords(bFirstInitAttempt);
 }
@@ -99,6 +113,8 @@ void FEngineSessionManager::InitializeRecords(bool bFirstAttempt)
 		// Get list of sessions in storage
 		if (StoredValuesLock.IsValid() && BeginReadWriteRecords())
 		{
+			UE_LOG(LogEngineSessionManager, Verbose, TEXT("Initializing EngineSessionManager for abnormal shutdown tracking"));
+
 			TArray<FSessionRecord> SessionRecordsToDelete;
 
 			// Attempt check each stored session
@@ -137,6 +153,8 @@ void FEngineSessionManager::InitializeRecords(bool bFirstAttempt)
 			EndReadWriteRecords();
 
 			bInitializedRecords = true;
+
+			UE_LOG(LogEngineSessionManager, Log, TEXT("EngineSessionManager initialized"));
 		}
 	}
 
@@ -151,7 +169,7 @@ void FEngineSessionManager::Tick(float DeltaTime)
 {
 	HeartbeatTimeElapsed += DeltaTime;
 
-	if (HeartbeatTimeElapsed > SessionManagerDefs::HeartbeatPeriodSeconds && !bShutdown)
+	if (HeartbeatTimeElapsed > (float)SessionManagerDefs::HeartbeatPeriodSeconds && !bShutdown)
 	{
 		HeartbeatTimeElapsed = 0.0f;
 
@@ -168,6 +186,13 @@ void FEngineSessionManager::Tick(float DeltaTime)
 			CurrentSession.Timestamp = FDateTime::UtcNow();
 
 			FPlatformMisc::SetStoredValue(SessionManagerDefs::StoreId, CurrentSessionSectionName, SessionManagerDefs::TimestampStoreKey, TimestampToString(CurrentSession.Timestamp));
+
+#if PLATFORM_SUPPORTS_WATCHDOG
+			if (!WatchdogSectionName.IsEmpty())
+			{
+				FPlatformMisc::SetStoredValue(SessionManagerDefs::StoreId, WatchdogSectionName, SessionManagerDefs::TimestampStoreKey, TimestampToString(CurrentSession.Timestamp));
+			}
+#endif
 		}
 	}
 }
@@ -180,6 +205,7 @@ void FEngineSessionManager::Shutdown()
 	FCoreDelegates::ApplicationWillEnterBackgroundDelegate.RemoveAll(this);
 	FCoreDelegates::ApplicationHasEnteredForegroundDelegate.RemoveAll(this);
 	FCoreDelegates::IsVanillaProductChanged.RemoveAll(this);
+	FSlateApplication::Get().GetOnModalLoopTickEvent().RemoveAll(this);
 
 	// Clear the session record for this session
 	if (bInitializedRecords)
@@ -196,6 +222,15 @@ void FEngineSessionManager::Shutdown()
 			FPlatformMisc::DeleteStoredValue(SessionManagerDefs::StoreId, CurrentSessionSectionName, SessionManagerDefs::BackgroundStoreKey);
 			FPlatformMisc::DeleteStoredValue(SessionManagerDefs::StoreId, CurrentSessionSectionName, SessionManagerDefs::UserActivityStoreKey);
 			FPlatformMisc::DeleteStoredValue(SessionManagerDefs::StoreId, CurrentSessionSectionName, SessionManagerDefs::VanillaStoreKey);
+
+#if PLATFORM_SUPPORTS_WATCHDOG
+			if (!WatchdogSectionName.IsEmpty())
+			{
+				FPlatformMisc::SetStoredValue(SessionManagerDefs::StoreId, WatchdogSectionName, SessionManagerDefs::StatusStoreKey, SessionManagerDefs::ShutdownSessionToken);
+				FPlatformMisc::SetStoredValue(SessionManagerDefs::StoreId, WatchdogSectionName, SessionManagerDefs::TimestampStoreKey, TimestampToString(FDateTime::UtcNow()));
+				WatchdogSectionName = FString();
+			}
+#endif
 		}
 
 		bInitializedRecords = false;
@@ -331,6 +366,7 @@ void FEngineSessionManager::DeleteStoredRecord(const FSessionRecord& Record)
  * @Trigger Fired only by the engine during startup, once for each "abnormal shutdown" detected that has not already been sent.
  *
  * @Type Static
+ * @Owner Chris.Wood
  *
  * @EventParam RunType - Editor or Game
  * @EventParam ProjectName - Project for the session that abnormally terminated. 
@@ -363,28 +399,23 @@ void FEngineSessionManager::DeleteStoredRecord(const FSessionRecord& Record)
  */
 void FEngineSessionManager::SendAbnormalShutdownReport(const FSessionRecord& Record)
 {
-#if PLATFORM_WINDOWS
-	const FString PlatformName(TEXT("Windows"));
-#elif PLATFORM_MAC
-	const FString PlatformName(TEXT("Mac"));
-#elif PLATFORM_LINUX
-	const FString PlatformName(TEXT("Linux"));
+	FString PlatformName(FPlatformProperties::PlatformName());
+
+#if PLATFORM_WINDOWS | PLATFORM_MAC | PLATFORM_LINUX
+	// do nothing
 #elif PLATFORM_PS4
-	const FString PlatformName(TEXT("PS4"));
 	if (Record.bIsDeactivated && !Record.bCrashed)
 	{
 		// Shutting down in deactivated state on PS4 is normal - don't report it
 		return;
 	}
 #elif PLATFORM_XBOXONE
-	const FString PlatformName(TEXT("XBoxOne"));
 	if (Record.bIsInBackground && !Record.bCrashed)
 	{
 		// Shutting down in background state on XB1 is normal - don't report it
 		return;
 	}
 #else
-	const FString PlatformName(TEXT("Unknown"));
 	return; // TODO: CWood: disabled on other platforms
 #endif
 
@@ -419,6 +450,8 @@ void FEngineSessionManager::SendAbnormalShutdownReport(const FSessionRecord& Rec
 	AbnormalShutdownAttributes.Add(FAnalyticsEventAttribute(TEXT("IsVanilla"), IsVanillaString));
 
 	FEngineAnalytics::GetProvider().RecordEvent(TEXT("Engine.AbnormalShutdown"), AbnormalShutdownAttributes);
+
+	UE_LOG(LogEngineSessionManager, Log, TEXT("EngineSessionManager sent abnormal shutdown report. Type=%s, SessionId=%s"), *ShutdownTypeString, *SessionIdString);
 }
 
 void FEngineSessionManager::CreateAndWriteRecordForSession()
@@ -463,6 +496,15 @@ void FEngineSessionManager::CreateAndWriteRecordForSession()
 	FPlatformMisc::SetStoredValue(SessionManagerDefs::StoreId, CurrentSessionSectionName, SessionManagerDefs::VanillaStoreKey, IsVanillaString);
 
 	SessionRecords.Add(CurrentSession);
+
+#if PLATFORM_SUPPORTS_WATCHDOG
+	bool bUseWatchdog = false;
+	GConfig->GetBool(TEXT("EngineSessionManager"), TEXT("UseWatchdogMTBF"), bUseWatchdog, GEngineIni);
+	if ((!CurrentSession.bIsDebugger && bUseWatchdog && !FParse::Param(FCommandLine::Get(), TEXT("NoWatchdog"))) || FParse::Param(FCommandLine::Get(), TEXT("ForceWatchdog")))
+	{
+		StartWatchdog(ModeString, CurrentSession.ProjectName, FPlatformProperties::PlatformName(), CurrentSession.SessionId, CurrentSession.EngineVersion);
+	}
+#endif
 }
 
 void FEngineSessionManager::OnCrashing()
@@ -471,6 +513,14 @@ void FEngineSessionManager::OnCrashing()
 	{
 		CurrentSession.bCrashed = true;
 		FPlatformMisc::SetStoredValue(SessionManagerDefs::StoreId, CurrentSessionSectionName, SessionManagerDefs::CrashStoreKey, SessionManagerDefs::TrueValueString);
+
+#if PLATFORM_SUPPORTS_WATCHDOG
+		if (!WatchdogSectionName.IsEmpty())
+		{
+			FPlatformMisc::SetStoredValue(SessionManagerDefs::StoreId, WatchdogSectionName, SessionManagerDefs::StatusStoreKey, SessionManagerDefs::CrashSessionToken);
+			FPlatformMisc::SetStoredValue(SessionManagerDefs::StoreId, WatchdogSectionName, SessionManagerDefs::TimestampStoreKey, TimestampToString(FDateTime::UtcNow()));
+		}
+#endif
 	}
 }
 
@@ -541,6 +591,14 @@ void FEngineSessionManager::OnUserActivity(const FUserActivity& UserActivity)
 	{
 		CurrentSession.CurrentUserActivity = GetUserActivityString();
 		FPlatformMisc::SetStoredValue(SessionManagerDefs::StoreId, CurrentSessionSectionName, SessionManagerDefs::UserActivityStoreKey, CurrentSession.CurrentUserActivity);
+
+#if PLATFORM_SUPPORTS_WATCHDOG
+		if (!WatchdogSectionName.IsEmpty())
+		{
+			FPlatformMisc::SetStoredValue(SessionManagerDefs::StoreId, WatchdogSectionName, SessionManagerDefs::UserActivityStoreKey, CurrentSession.CurrentUserActivity);
+			FPlatformMisc::SetStoredValue(SessionManagerDefs::StoreId, WatchdogSectionName, SessionManagerDefs::TimestampStoreKey, TimestampToString(FDateTime::UtcNow()));
+		}
+#endif
 	}
 }
 
@@ -556,4 +614,98 @@ FString FEngineSessionManager::GetUserActivityString() const
 	return UserActivity.ActionName;
 }
 
+#if PLATFORM_SUPPORTS_WATCHDOG
+
+/**
+ * @EventName Engine.StartWatchdog
+ *
+ * @Trigger Event raised by EngineSessionManager as part of MTBF tracking. Records an attempt to start the UnrealWatchdog process.
+ *
+ * @Type Static
+ * @Owner Chris.Wood
+ *
+ * @EventParam RunType - Editor or Game
+ * @EventParam ProjectName - Project for the session.
+ * @EventParam Platform - Windows, Mac, Linux
+ * @EventParam SessionId - Analytics SessionID of the session.
+ * @EventParam EngineVersion - EngineVersion of the session.
+ * @EventParam IsInternalBuild - internal Epic build environment or not? Calls FEngineBuildSettings::IsInternalBuild(). Value is Yes or No.
+ * @EventParam Outcome - Whether the watchdog was started successfully. One of Succeeded, CreateProcFailed or MissingBinaryFailed.
+ *
+ * @Comments Currently only runs Watchdog when MTBF is enabled, we aren't debugging, we're a DESKTOP platform and watchdog is specifically enabled via config or command line arg.
+ */
+void FEngineSessionManager::StartWatchdog(const FString& RunType, const FString& ProjectName, const FString& PlatformName, const FString& SessionId, const FString& EngineVersion)
+{
+	uint32 ProcessId =  FPlatformProcess::GetCurrentProcessId();
+	const int SuccessfulRtnCode = 0;	// hardcode this for now, zero might not always be correct
+
+	FString WatchdogClientArguments =
+		FString::Printf(TEXT("-PID=%u -RunType=%s -ProjectName=\"%s\" -Platform=%s -SessionId=%s -EngineVersion=%s -SuccessfulRtnCode=%d -HeartbeatSeconds=%d"), ProcessId, *RunType, *ProjectName, *PlatformName, *SessionId, *EngineVersion, SuccessfulRtnCode, SessionManagerDefs::HeartbeatPeriodSeconds);
+
+	if (FEngineBuildSettings::IsInternalBuild())
+	{
+		// Suppress the watchdog dialogs if this engine session should never show interactive UI
+		if (!FApp::IsUnattended() && !IsRunningDedicatedServer() && FApp::CanEverRender())
+		{
+			// Only show watchdog dialogs if it's set in config
+			bool bAllowWatchdogDialogs = false;
+			GConfig->GetBool(TEXT("EngineSessionManager"), TEXT("AllowWatchdogDialogs"), bAllowWatchdogDialogs, GEngineIni);
+
+			if (bAllowWatchdogDialogs)
+			{
+				WatchdogClientArguments.Append(TEXT(" -AllowDialogs"));
+			}
+		}
+	}
+
+	FString WatchdogPath = FPaths::ConvertRelativePathToFull(FPlatformProcess::GenerateApplicationPath(TEXT("UnrealWatchdog"), EBuildConfigurations::Development));
+
+	TArray< FAnalyticsEventAttribute > WatchdogStartedAttributes;
+	WatchdogStartedAttributes.Add(FAnalyticsEventAttribute(TEXT("RunType"), RunType));
+	WatchdogStartedAttributes.Add(FAnalyticsEventAttribute(TEXT("ProjectName"), ProjectName));
+	WatchdogStartedAttributes.Add(FAnalyticsEventAttribute(TEXT("Platform"), PlatformName));
+	WatchdogStartedAttributes.Add(FAnalyticsEventAttribute(TEXT("SessionId"), SessionId));
+	WatchdogStartedAttributes.Add(FAnalyticsEventAttribute(TEXT("IsInternalBuild"), FEngineBuildSettings::IsInternalBuild() ? TEXT("Yes") : TEXT("No")));
+
+	if (FPaths::FileExists(WatchdogPath))
+	{
+		FProcHandle WatchdogProcessHandle = FPlatformProcess::CreateProc(*WatchdogPath, *WatchdogClientArguments, true, true, false, NULL, 0, NULL, NULL);
+
+		if (WatchdogProcessHandle.IsValid())
+		{
+			FString WatchdogStartTimeString = TimestampToString(FDateTime::UtcNow());
+
+			WatchdogStartedAttributes.Add(FAnalyticsEventAttribute(TEXT("Outcome"), TEXT("Succeeded")));
+			UE_LOG(LogEngineSessionManager, Log, TEXT("Started UnrealWatchdog for process id %u"), ProcessId);
+
+			WatchdogSectionName = GetWatchdogStoreSectionString(ProcessId);
+
+			FPlatformMisc::SetStoredValue(SessionManagerDefs::StoreId, WatchdogSectionName, SessionManagerDefs::CommandLineStoreKey, FCommandLine::GetOriginalForLogging());
+			FPlatformMisc::SetStoredValue(SessionManagerDefs::StoreId, WatchdogSectionName, SessionManagerDefs::StartupTimeStoreKey, WatchdogStartTimeString);
+			FPlatformMisc::SetStoredValue(SessionManagerDefs::StoreId, WatchdogSectionName, SessionManagerDefs::TimestampStoreKey, WatchdogStartTimeString);
+			FPlatformMisc::SetStoredValue(SessionManagerDefs::StoreId, WatchdogSectionName, SessionManagerDefs::StatusStoreKey, SessionManagerDefs::RunningSessionToken);
+			FPlatformMisc::SetStoredValue(SessionManagerDefs::StoreId, WatchdogSectionName, SessionManagerDefs::UserActivityStoreKey, CurrentSession.CurrentUserActivity);
+		}
+		else
+		{
+			WatchdogStartedAttributes.Add(FAnalyticsEventAttribute(TEXT("Outcome"), TEXT("CreateProcFailed")));
+			UE_LOG(LogEngineSessionManager, Warning, TEXT("Unable to start UnrealWatchdog.exe. CreateProc failed."));
+		}
+	}
+	else
+	{
+		WatchdogStartedAttributes.Add(FAnalyticsEventAttribute(TEXT("Outcome"), TEXT("MissingBinaryFailed")));
+		UE_LOG(LogEngineSessionManager, Warning, TEXT("Unable to start UnrealWatchdog.exe. File not found."));
+	}
+
+	FEngineAnalytics::GetProvider().RecordEvent(TEXT("Engine.StartWatchdog"), WatchdogStartedAttributes);
+}
+
+FString FEngineSessionManager::GetWatchdogStoreSectionString(uint32 InPID)
+{
+	return FString::Printf(TEXT("%s%s/%u"), *SessionManagerDefs::WatchdogRecordSectionPrefix, *SessionManagerDefs::WatchdogVersionString, InPID);
+}
+
+#endif // PLATFORM_SUPPORTS_WATCHDOG
+
 #undef LOCTEXT_NAMESPACE
diff --git a/Engine/Source/Runtime/Engine/Private/EngineSessionManager.h b/Engine/Source/Runtime/Engine/Private/EngineSessionManager.h
--- a/Engine/Source/Runtime/Engine/Private/EngineSessionManager.h
+++ b/Engine/Source/Runtime/Engine/Private/EngineSessionManager.h
@@ -4,6 +4,10 @@
 
 #include "CoreMinimal.h"
 
+// TODO: implement Watchdog on Mac and Linux and replace Windows check with desktop check below
+//#define PLATFORM_SUPPORTS_WATCHDOG		PLATFORM_DESKTOP
+#define PLATFORM_SUPPORTS_WATCHDOG		PLATFORM_WINDOWS
+
 struct FUserActivity;
 
 enum class EEngineSessionManagerMode
@@ -67,6 +71,11 @@ private:
 	void OnVanillaStateChanged(bool bIsVanilla);
 	FString GetUserActivityString() const;
 
+#if PLATFORM_SUPPORTS_WATCHDOG
+	void StartWatchdog(const FString& RunType, const FString& ProjectName, const FString& PlatformName, const FString& SessionId, const FString& EngineVersion);
+	FString GetWatchdogStoreSectionString(uint32 InPID);
+#endif
+
 private:
 	EEngineSessionManagerMode Mode;
 	bool bInitializedRecords;
@@ -75,4 +84,8 @@ private:
 	FSessionRecord CurrentSession;
 	FString CurrentSessionSectionName;
 	TArray<FSessionRecord> SessionRecords;
+
+#if PLATFORM_SUPPORTS_WATCHDOG
+	FString WatchdogSectionName;
+#endif
 };
diff --git a/Engine/Source/Runtime/Engine/Public/EngineLogs.h b/Engine/Source/Runtime/Engine/Public/EngineLogs.h
--- a/Engine/Source/Runtime/Engine/Public/EngineLogs.h
+++ b/Engine/Source/Runtime/Engine/Public/EngineLogs.h
@@ -31,3 +31,4 @@ ENGINE_API DECLARE_LOG_CATEGORY_EXTERN(LogSubtitle, Log, All);
 ENGINE_API DECLARE_LOG_CATEGORY_EXTERN(LogTexture, Log, All);
 ENGINE_API DECLARE_LOG_CATEGORY_EXTERN(LogPlayerManagement, Error, All);
 ENGINE_API DECLARE_LOG_CATEGORY_EXTERN(LogSecurity, Warning, All);
+ENGINE_API DECLARE_LOG_CATEGORY_EXTERN(LogEngineSessionManager, Log, All);
\ No newline at end of file
diff --git a/Engine/Source/Runtime/Slate/Private/Framework/Application/SlateApplication.cpp b/Engine/Source/Runtime/Slate/Private/Framework/Application/SlateApplication.cpp
--- a/Engine/Source/Runtime/Slate/Private/Framework/Application/SlateApplication.cpp
+++ b/Engine/Source/Runtime/Slate/Private/Framework/Application/SlateApplication.cpp
@@ -2014,14 +2014,17 @@ void FSlateApplication::AddModalWindow( TSharedRef<SWindow> InSlateWindow, const
 				FPlatformProcess::Sleep(SleepTime);
 			}
 
+			const float DeltaTime = GetDeltaTime();
+
+			// Tick any other systems that need to update during modal dialogs
+			ModalLoopTickEvent.Broadcast(DeltaTime);
+
 			FPlatformMisc::BeginNamedEvent(FColor::Magenta, "Slate::Tick");
 
 			{
 				SCOPE_CYCLE_COUNTER(STAT_SlateTickTime);
 				SLATE_CYCLE_COUNTER_SCOPE(GSlateTotalTickTime);
 
-				const float DeltaTime = GetDeltaTime();
-
 				// Tick and pump messages for the platform.
 				TickPlatform(DeltaTime);
 
diff --git a/Engine/Source/Runtime/Slate/Public/Framework/Application/SlateApplication.h b/Engine/Source/Runtime/Slate/Public/Framework/Application/SlateApplication.h
--- a/Engine/Source/Runtime/Slate/Public/Framework/Application/SlateApplication.h
+++ b/Engine/Source/Runtime/Slate/Public/Framework/Application/SlateApplication.h
@@ -1430,6 +1430,17 @@ public:
 	*/
 	FDelegateHandle RegisterOnWindowActionNotification(const FOnWindowAction& Notification);
 
+
+	/** Event type for when Slate is ticking during a modal dialog loop */
+	DECLARE_EVENT_OneParam(FSlateApplication, FOnModalLoopTickEvent, float);
+
+	/**
+	 * Get the FOnModalLoopTickEvent for the Slate Application. Allows clients to register for callbacks during modal dialog loops.
+	 *
+	 * @return The application's FOnModalLoopTickEvent.
+	 */
+	FOnModalLoopTickEvent& GetOnModalLoopTickEvent() { return ModalLoopTickEvent; }
+
 	/**
 	* Unregister the notification because it is no longer desired.
 	*
@@ -1951,6 +1962,9 @@ private:
 	/** Delegate for post slate Tick */
 	FSlateTickEvent PostTickEvent;
 
+	/** Delegate for slate Tick during modal dialogs */
+	FOnModalLoopTickEvent ModalLoopTickEvent;
+
 	/** Critical section to avoid multiple threads calling Slate Tick when we're synchronizing between the Slate Loading Thread and the Game Thread. */
 	FCriticalSection SlateTickCriticalSection;
 
