Fix missing UnrealWatchdog include causing build failures in some development streams.

diff --git a/Engine/Source/Programs/UnrealWatchdog/Private/UnrealWatchdog.h b/Engine/Source/Programs/UnrealWatchdog/Private/UnrealWatchdog.h
--- a/Engine/Source/Programs/UnrealWatchdog/Private/UnrealWatchdog.h
+++ b/Engine/Source/Programs/UnrealWatchdog/Private/UnrealWatchdog.h
@@ -3,6 +3,7 @@
 #pragma once
 
 #include "CoreMinimal.h"
+#include "HAL/PlatformProcess.h"
 
 DECLARE_LOG_CATEGORY_EXTERN(UnrealWatchdogLog, Log, All);
 
